CIS fix - removed shadow variable

diff --git a/Engine/Plugins/Experimental/DataInterface/Source/DataInterface/Public/DataInterfaceContext.h b/Engine/Plugins/Experimental/DataInterface/Source/DataInterface/Public/DataInterfaceContext.h
--- a/Engine/Plugins/Experimental/DataInterface/Source/DataInterface/Public/DataInterfaceContext.h
+++ b/Engine/Plugins/Experimental/DataInterface/Source/DataInterface/Public/DataInterfaceContext.h
@@ -197,14 +197,14 @@ private:
 
 	FParam::EFlags GetFlagsForConstructor(const FContext& InContext, ValueType* InValuePtrToWrap) const
 	{
-		FParam::EFlags Flags = FParam::EFlags::None;
+		FParam::EFlags NewFlags = FParam::EFlags::None;
 		
 		// Add chunked flags or not
 		if constexpr (TModels<Private::CSizedContainerWithAccessibleDataAsRawPtr, ValueType>::Value)
 		{
 			// Containers always assume chunked
 			checkSlow(InContext.GetChunkSize() == InValuePtrToWrap->Num());
-			Flags |= InContext.GetChunkFlags();
+			NewFlags |= InContext.GetChunkFlags();
 		}
 		else
 		{
@@ -214,10 +214,10 @@ private:
 		// Add const flags or not
 		if constexpr (!TIsConst<ValueType>::Value)
 		{
-			Flags |= FParam::EFlags::Mutable;
+			NewFlags |= FParam::EFlags::Mutable;
 		}
 
-		return Flags;
+		return NewFlags;
 	}
 };
 
