Workaround for VS2022 preview compiler bug

#rnx

diff --git a/Engine/Plugins/Experimental/DataInterface/Source/DataInterface/Public/DataInterfaceKernel.h b/Engine/Plugins/Experimental/DataInterface/Source/DataInterface/Public/DataInterfaceKernel.h
--- a/Engine/Plugins/Experimental/DataInterface/Source/DataInterface/Public/DataInterfaceKernel.h
+++ b/Engine/Plugins/Experimental/DataInterface/Source/DataInterface/Public/DataInterfaceKernel.h
@@ -48,7 +48,8 @@ public:
 					{
 						if constexpr (TIsDerivedFrom<ParamTypes, FParam>::Value)
 						{
-							return static_cast<typename ParamTypes::TValueType*>(InArg.Data);
+							using Type = decltype(typename ParamTypes::TValueType{});
+							return static_cast<Type*>(InArg.Data);
 						}
 						else
 						{
