Data Interfaces experimental plugin update :
- Added Context branching support (up to 64 elements)
- Added TParam<Type> operator [] and Opertator Type(), so we can operate directly with the TParams (this is used instead of the Kernel::run, which is harder to debug and deal with)
- Fixed batches and passing a single element to a Context that expects N elements (no more memory stomping)
- Added new ways to obtain a TParam from the Context (GetParamPtr, GetParamAs, GetParameter with TOptional). Also added a GetResultAs<Type>, to ease working with the parameters
- Fixed GetDataSafe to support containers (with Num() and GetData)
- Added NumElements to FParam (indicates the number of elements that the Data holds for this parameter)
- Updated FParam CanAssignWith to allow copy with less restrictions if the target param is empty / unasigned
- Added more constructors to TWrapParam, one with explicit NumElements and other with an ArrayView
- Added more constructors to TParam, with explicit NumElements
- Added EngineTest for DataInterface with a example at Animation.DataInterfaces.XXXXXX
- Prototype of parameter storage using arenas : Context requests the size of the passed paramters and copies the ones selected by the user to be transfered
- Prototype of stack like parameter storage, which allows storing small values directly into the Data pointer of FParam struct (pending see if it can be done constexpr, now is checked at runtime)

[REDACTED_HANDLE] [REDACTED_HANDLE] [REDACTED_HANDLE]

diff --git a/Engine/Plugins/Experimental/DataInterface/Source/DataInterface/Private/DataInterface.cpp b/Engine/Plugins/Experimental/DataInterface/Source/DataInterface/Private/DataInterface.cpp
--- a/Engine/Plugins/Experimental/DataInterface/Source/DataInterface/Private/DataInterface.cpp
+++ b/Engine/Plugins/Experimental/DataInterface/Source/DataInterface/Private/DataInterface.cpp
@@ -10,7 +10,7 @@ namespace UE::DataInterface::Private
 static bool CheckCompatibility(const IDataInterface* InDataInterface, const UE::DataInterface::FContext& InContext)
 {
 	check(InDataInterface);
-	return InContext.GetResultRaw().GetType().GetName() == InDataInterface->GetReturnTypeName();
+	return InContext.GetResultParam().GetType().GetName() == InDataInterface->GetReturnTypeName();
 }
 
 }
@@ -25,6 +25,17 @@ bool IDataInterface::GetDataIfCompatibleInternal(const UE::DataInterface::FConte
 	return false;
 }
 
+bool IDataInterface::GetData(const UE::DataInterface::FContext& Context) const
+{
+	return GetDataIfCompatibleInternal(Context);
+}
+
+bool IDataInterface::GetDataChecked(const UE::DataInterface::FContext& Context) const
+{
+	check(UE::DataInterface::Private::CheckCompatibility(this, Context));
+	return GetDataRawInternal(Context);
+}
+
 bool IDataInterface::GetData(const UE::DataInterface::FContext& Context, UE::DataInterface::FParam& OutResult) const
 {
 	const UE::DataInterface::FContext CallingContext = Context.WithResult(OutResult);
diff --git a/Engine/Plugins/Experimental/DataInterface/Source/DataInterface/Private/DataInterfaceContext.cpp b/Engine/Plugins/Experimental/DataInterface/Source/DataInterface/Private/DataInterfaceContext.cpp
--- a/Engine/Plugins/Experimental/DataInterface/Source/DataInterface/Private/DataInterfaceContext.cpp
+++ b/Engine/Plugins/Experimental/DataInterface/Source/DataInterface/Private/DataInterfaceContext.cpp
@@ -34,56 +34,168 @@ const FContext& FThreadContext::Get()
 	return *ContextData.ContextStack.Top();
 }
 
-FContext::FContext(float InDeltaTime, FState& InState, IDataInterfaceParameters* InParameters)
-	: State(InState)
-	, Result(nullptr)
+FContext::FContext(float InDeltaTime, FState& InState, FParamStorage& InParamStorage, IDataInterfaceParameters* InParameters)
+	: State(&InState)
+	, ParamStorage(&InParamStorage)
 	, Parameters(InParameters)
 	, DeltaTime(InDeltaTime)
+	, BranchingMask(MAX_uint64)
 {
 }
 
+FContext::~FContext()
+{
+	// Remove any handles before the storage goes out of scope
+	AdditionalParameterHandles.Empty();
+
+	if (BlockHandle != InvalidBlockHandle)
+	{
+		ParamStorage->ReleaseBlock(BlockHandle);
+	}
+}
+
 FContext FContext::WithResult(FParam& InResult) const
 {
-	FContext NewContext(DeltaTime, State, &InResult);
+	FContext NewContext(DeltaTime, *State, *ParamStorage, InResult);
 	NewContext.Parent = this;
 	NewContext.Root = (Root == nullptr) ? this : Root;
 	NewContext.CallstackHash = CallstackHash;
+	NewContext.BranchingMask = BranchingMask;
 
 	return NewContext;
 }
 
 FContext FContext::WithParameter(FName ParameterId, const FParam& InParameter) const
 {
-	FContext NewContext(DeltaTime, State, Result);
-	NewContext.AdditionalParameters.Add(ParameterId, InParameter);
+	FContext NewContext(DeltaTime, *State, *ParamStorage, *Result);
 	NewContext.Parent = this;
 	NewContext.Root = (Root == nullptr) ? this : Root;
 	NewContext.CallstackHash = CallstackHash;
-
+	NewContext.BranchingMask = BranchingMask;
+	NewContext.AddParameters({ TPairInitializer(ParameterId, InParameter) });
 	return NewContext;
 }
 
 FContext FContext::WithParameters(TArrayView<const TPair<FName, FParam>> InParameters) const
 {
-	FContext NewContext(DeltaTime, State, Result);
-	for(const TPair<FName, FParam>& ParamPair : InParameters)
-	{
-		NewContext.AdditionalParameters.Add(ParamPair.Key, ParamPair.Value);
-	}
+	FContext NewContext(DeltaTime, *State, *ParamStorage, *Result);
 	NewContext.Parent = this;
 	NewContext.Root = (Root == nullptr) ? this : Root;
 	NewContext.CallstackHash = CallstackHash;
+	NewContext.BranchingMask = BranchingMask;
+	NewContext.AddParameters(InParameters);
 
 	return NewContext;
 }
 
+FContext FContext::WithResultAndParameters(FParam& InResult, TArrayView<const TPair<FName, FParam>> InParameters) const
+{
+	FContext NewContext(DeltaTime, *State, *ParamStorage, InResult);
+	NewContext.Parent = this;
+	NewContext.Root = (Root == nullptr) ? this : Root;
+	NewContext.CallstackHash = CallstackHash;
+	NewContext.BranchingMask = BranchingMask;
+	NewContext.AddParameters(InParameters);
+
+	return NewContext;
+}
+
+int32 FContext::GetParametersSize(TArrayView<const TPair<FName, FParam>> InParameters, TArray<int32> &ParamAllocSizes) const
+{
+	int32 TotalParamAllocSize = 0;
+
+	const int32 NumParameters = InParameters.Num();
+
+	for (int i = 0; i < NumParameters; i++)
+	{
+		const TPair<FName, FParam>& ParamPair = InParameters[i];
+		const FParam& SourceParam = ParamPair.Value;
+
+		if (EnumHasAnyFlags(SourceParam.GetFlags(), FParam::EFlags::Stored))
+		{
+			const FParamType& ParamType = SourceParam.GetType();
+
+			const int32 NumElem = SourceParam.GetNumElements();
+			const int32 ParamAlignment = ParamType.GetAlignment();
+			const int32 ParamSize = ParamType.GetSize();
+
+			const int32 ParamAllocSize = NumElem * Align(ParamSize, ParamAlignment);
+
+			ParamAllocSizes[i] = ParamAllocSize;
+			TotalParamAllocSize += ParamAllocSize;
+		}
+	}
+
+	return TotalParamAllocSize;
+}
+
+void FContext::AddParameters(TArrayView<const TPair<FName, FParam>> InParameters)
+{
+	int32 TotalParamAllocSize = 0;
+
+	const int32 NumParameters = InParameters.Num();
+
+	TArray<int32> ParamAllocSizes;
+	ParamAllocSizes.Reserve(NumParameters);
+	ParamAllocSizes.AddZeroed(NumParameters);
+
+	TotalParamAllocSize = GetParametersSize(InParameters, ParamAllocSizes);
+
+	uint8* TargetMemory = nullptr;
+	check(BlockHandle == InvalidBlockHandle);
+
+	if (TotalParamAllocSize > 0)
+	{
+		const FParamStorage::TBlockDataPair BlockData = ParamStorage->RequestBlock(TotalParamAllocSize);
+
+		BlockHandle = BlockData.Key;
+		TargetMemory = BlockData.Value;
+	}
+
+	for(int i = 0; i < NumParameters; i++)
+	{
+		const TPair<FName, FParam>& ParamPair = InParameters[i];
+		const FParam& SourceParam = ParamPair.Value;
+		const int32 ParamAllocatedSize = ParamAllocSizes[i];
+
+		if (ParamAllocatedSize > 0 && EnumHasAnyFlags(SourceParam.GetFlags(), FParam::EFlags::Stored))
+		{
+			const FParamType& ParamType = ParamPair.Value.GetType();
+
+			const FParamType::ParamCopyFunction ParamCopyFunction = ParamType.GetParamCopyFunction();
+
+			FParam ClonedParam = ParamCopyFunction(SourceParam, TargetMemory, ParamAllocatedSize);
+
+			TargetMemory += ParamAllocatedSize;
+
+			AdditionalParameters.Add(ParamPair.Key, ClonedParam);
+		}
+		else
+		{
+			AdditionalParameters.Add(ParamPair.Key, ParamPair.Value);
+		}
+	}
+}
+
 FContext FContext::WithParameters(IDataInterfaceParameters* InParameters) const
 {
-	FContext NewContext(DeltaTime, State, Result);
+	FContext NewContext(DeltaTime, *State, *ParamStorage, *Result);
 	NewContext.Parameters = InParameters;
 	NewContext.Parent = this;
 	NewContext.Root = (Root == nullptr) ? this : Root;
 	NewContext.CallstackHash = CallstackHash;
+	NewContext.BranchingMask = BranchingMask;
+
+	return NewContext;
+}
+
+FContext FContext::CreateSubContext() const
+{
+	FContext NewContext(DeltaTime, *State, *ParamStorage, *Result);
+	NewContext.Parent = this;
+	NewContext.Root = (Root == nullptr) ? this : Root;
+	NewContext.CallstackHash = CallstackHash;
+	NewContext.BranchingMask = BranchingMask;
 
 	return NewContext;
 }
@@ -117,22 +229,40 @@ bool FContext::GetParameter(FName InKey, FParam& OutParam) const
 				}
 			}
 		}
+		// Check additional parameter handles
+		// Note this is currently 'else if' because we dont allow creating new contexts with both 
+		// AdditionalParameters and AdditionalParametersHandles
+		else if (!CurrentContext->AdditionalParameterHandles.IsEmpty())
+		{
+			// Find the parameter
+			if (const FHParam* HParam = CurrentContext->AdditionalParameterHandles.Find(InKey))
+			{
+				const FParam* Param = ParamStorage->GetParam(HParam->ParamHandle);
+				if (Param != nullptr && Param->CanAssignTo(OutParam))
+				{
+					OutParam = Param;
+					return true;
+				}
+			}
+		}
 	}
 
 	return false;
 }
 
+
 TParam<const float> FContext::GetDeltaTimeParam() const
 {
-	return TWrapParam<const float>(*this, &DeltaTime);
+	return TWrapParam<const float>(&DeltaTime);
 }
 
 FContext FContext::WithCallRaw(const IDataInterface* InDataInterface) const
 {
-	FContext NewContext(DeltaTime, State, Result);
+	FContext NewContext(DeltaTime, *State, *ParamStorage, *Result);
 	NewContext.Parent = this;
 	NewContext.Root = (Root == nullptr) ? this : Root;
 	NewContext.CallstackHash = HashCombineFast(CallstackHash, GetTypeHash(InDataInterface));
+	NewContext.BranchingMask = BranchingMask;
 	return NewContext;
 }
 
@@ -140,7 +270,7 @@ void FContext::FlushRelevancy() const
 {
 	// Flush all relevancy-based allocations if they were not used this update
 	// @TODO: this could be more efficient if we use a more linear iteration here 
-	for(auto It = State.RelevancyValueMap.CreateIterator(); It; ++It)
+	for(auto It = State->RelevancyValueMap.CreateIterator(); It; ++It)
 	{
 		if(It->Value.UpdateCounter != UpdateCounter)
 		{
diff --git a/Engine/Plugins/Experimental/DataInterface/Source/DataInterface/Private/DataInterfaceParam.cpp b/Engine/Plugins/Experimental/DataInterface/Source/DataInterface/Private/DataInterfaceParam.cpp
--- a/Engine/Plugins/Experimental/DataInterface/Source/DataInterface/Private/DataInterfaceParam.cpp
+++ b/Engine/Plugins/Experimental/DataInterface/Source/DataInterface/Private/DataInterfaceParam.cpp
@@ -1,6 +1,7 @@
 ﻿// Copyright Epic Games, Inc. All Rights Reserved.
 
 #include "DataInterfaceParam.h"
+#include "DataInterfaceParamStorage.h"
 
 namespace UE::DataInterface
 {
@@ -11,6 +12,7 @@ static TArray<TUniqueFunction<void(void)>> GDeferredTypes;
 
 FParam::FParam(const FParamType& InType, void* InData, EFlags InFlags)
 	: Data(InData)
+	, NumElements((InData != nullptr) ? 1 : 0)
 	, TypeId(InType.GetTypeId())
 	, Flags(InFlags)
 {
@@ -18,19 +20,38 @@ FParam::FParam(const FParamType& InType, void* InData, EFlags InFlags)
 	check(Data);
 }
 
+FParam::FParam(const FParamType& InType, void* InData, int32 InNumElements, EFlags InFlags)
+	: Data(InData)
+	, NumElements(InNumElements)
+	, TypeId(InType.GetTypeId())
+	, Flags(InFlags)
+{
+}
+
+FParam::FParam(const FParamType& InType, EFlags InFlags)
+	: Data(nullptr)
+	, NumElements(0)
+	, TypeId(InType.GetTypeId())
+	, Flags(InFlags)
+{
+	check(TypeId != INDEX_NONE);
+}
+
 FParam::FParam(const FParam* InOtherParam)
-	: Data(InOtherParam->Data)
+	: Data(EnumHasAnyFlags(InOtherParam->Flags, FParam::EFlags::Embedded) ? (void*)&InOtherParam->Data : InOtherParam->Data)
+	, NumElements(InOtherParam->NumElements)
 	, TypeId(InOtherParam->TypeId)
 	, Flags(InOtherParam->Flags)
 {
+	EnumRemoveFlags(Flags, FParam::EFlags::Embedded); // For compatibility reasons, I can not keep the embedded value if we create a copy of the FParam
 }
 
 bool FParam::CanAssignTo(const FParam& InParam) const
 {
-	return CanAssignWith(InParam.GetType(), InParam.Flags);
+	return CanAssignWith(InParam.GetType(), InParam.Flags, InParam.GetNumElements());
 }
 
-bool FParam::CanAssignWith(const FParamType& InType, EFlags InFlags, FStringBuilderBase* OutReasonPtr) const
+bool FParam::CanAssignWith(const FParamType& InType, EFlags InFlags, int32 InNumElements, FStringBuilderBase* OutReasonPtr) const
 {
 	// Check type
 	if(TypeId != InType.GetTypeId())
@@ -43,8 +64,8 @@ bool FParam::CanAssignWith(const FParamType& InType, EFlags InFlags, FStringBuil
 		return false;
 	}
 	
-	// Check const
-	if(!EnumHasAnyFlags(InFlags, EFlags::Mutable))
+	// Check const if not empty, else we assume it is an empty param and can be written, mutable or not
+	if(InNumElements > 0 && !EnumHasAnyFlags(InFlags, EFlags::Mutable))
 	{
 		if(OutReasonPtr)
 		{
@@ -53,29 +74,35 @@ bool FParam::CanAssignWith(const FParamType& InType, EFlags InFlags, FStringBuil
 		return false;
 	}
 
-	// Check chunking
-	if(EnumHasAnyFlags(Flags, EFlags::Chunked) != EnumHasAnyFlags(InFlags, EFlags::Chunked))
+	// Check batching
+	if (EnumHasAnyFlags(Flags, EFlags::Batched) != EnumHasAnyFlags(InFlags, EFlags::Batched))
 	{
-		if(OutReasonPtr)
+		// If different batching, allow it if the target is empty and the source is 1 (this enables passing a single value to a context that expects N)
+		if (NumElements != 1 || InNumElements != 0)
 		{
-			OutReasonPtr->Append(TEXT("Incompatible chunking"));
+			if (OutReasonPtr)
+			{
+				OutReasonPtr->Append(TEXT("Incompatible batching"));
+			}
+			return false;
 		}
-		return false;
 	}
 	
 	return true;
 }
 
-void FParamType::FRegistrar::RegisterType(FParamType& InType, const UScriptStruct* InStruct, FName InName, uint32 InSize, uint32 InAlignment)
+void FParamType::FRegistrar::RegisterType(FParamType& InType, const UScriptStruct* InStruct, FName InName, uint32 InSize, uint32 InAlignment, const ParamCopyFunction InParamCopyFunctionPtr)
 {
 	InType.Struct = InStruct;
 	InType.Name = InName;
 	InType.Size = InSize;
 	InType.Alignment = InAlignment;
+	InType.ParamCopyFunctionPtr = InParamCopyFunctionPtr;
 
 	check(InType.Name != NAME_None);
 	check(InType.Size != 0);
 	check(InType.Alignment != 0);
+	check(InType.ParamCopyFunctionPtr != nullptr);
 	
 	InType.TypeId = GTypeRegistry.Add(&InType);
 }
@@ -140,4 +167,30 @@ FPropertyParam::FPropertyParam(const FProperty* InProperty, UObject* InContainer
 {
 }
 */
+
+FHParam::FHParam(FParamStorage* InOwnerStorage, FParamHandle InParamHandle)
+	: OwnerStorage(InOwnerStorage)
+	, ParamHandle(InParamHandle)
+{
+}
+
+FHParam::FHParam(const FHParam& Other)
+	: OwnerStorage(Other.OwnerStorage)
+	, ParamHandle(Other.ParamHandle)
+{
+	if (ParamHandle != InvalidParamHandle)
+	{
+		OwnerStorage->IncRefCount(ParamHandle);
+	}
 }
+
+FHParam::~FHParam()
+{
+	if (ParamHandle != InvalidParamHandle)
+	{
+		check(OwnerStorage != nullptr);
+		OwnerStorage->DecRefCount(ParamHandle);
+	}
+}
+
+} // end namespace UE::DataInterface
diff --git a/Engine/Plugins/Experimental/DataInterface/Source/DataInterface/Private/DataInterfaceParamStorage.cpp b/Engine/Plugins/Experimental/DataInterface/Source/DataInterface/Private/DataInterfaceParamStorage.cpp
new file mode 100644
--- /dev/null
+++ b/Engine/Plugins/Experimental/DataInterface/Source/DataInterface/Private/DataInterfaceParamStorage.cpp
@@ -0,0 +1,157 @@
+// Copyright Epic Games, Inc. All Rights Reserved.
+
+#include "DataInterfaceParamStorage.h"
+#include "HAL/UnrealMemory.h"
+
+namespace UE::DataInterface
+{
+
+
+FParamStorage::FParamStorage(int32 InMaxParams, int32 InAllocatedMemorySize, int32 InMaxBlocks, bool InAllowGrowing)
+	: AllowGrowing(InAllowGrowing)
+{
+	check(InMaxParams > 0 && InMaxBlocks > 0 && InAllocatedMemorySize > 0);
+
+	Parameters.Reserve(InMaxParams);
+	Parameters.AddDefaulted(InMaxParams);
+
+	RefCounts.Reserve(InMaxParams);
+	RefCounts.AddZeroed(InMaxParams);
+
+	ParamMemoryHandles.Reserve(InMaxParams);
+	for (int32 i = 0; i < InMaxParams; i++)
+	{
+		ParamMemoryHandles.Add(InvalidBlockHandle);
+	}
+
+	// --- ---
+
+	RawMemory = (uint8*)FMemory::Malloc(InAllocatedMemorySize);
+	RawMemorySize = InAllocatedMemorySize;
+
+	BlockOffets.Reserve(InMaxBlocks);
+	BlockOffets.AddZeroed(InMaxBlocks);
+
+	BlockSizes.Reserve(InMaxBlocks);
+	BlockSizes.AddZeroed(InMaxBlocks);
+
+	BlockFreeFlag.Reserve(InMaxBlocks);
+	BlockFreeFlag.Init(true, InMaxBlocks);
+
+	FreeBlockIndex = 0;
+}
+
+FParamStorage::~FParamStorage()
+{
+	// There should be no blocks of memory or HParams in memory at this point
+	check(FreeBlockIndex == 0);
+	check(ParamFreeIndex == 0);
+
+	FMemory::Free(RawMemory);
+}
+
+FParamStorage::TBlockDataPair FParamStorage::RequestBlock(int32 RequestedBlockSize)
+{
+	FParamStorageHandle BlockHandle = InvalidBlockHandle;
+	uint8* BlockMemory = nullptr;
+
+	check((CurrentBlocksSize + RequestedBlockSize < RawMemorySize) || AllowGrowing);
+	check((FreeBlockIndex >= 0 && FreeBlockIndex < (BlockSizes.Num() - 1)) || AllowGrowing);
+
+	if (CheckBufferSize(RequestedBlockSize) && CheckNumBlocks(FreeBlockIndex))
+	{
+		const int32 BlockIndex = FreeBlockIndex;
+
+		BlockOffets[BlockIndex] = (BlockIndex > 0) ? BlockOffets[BlockIndex - 1] + BlockSizes[BlockIndex - 1] : 0;
+		BlockSizes[BlockIndex] = RequestedBlockSize;
+		CurrentBlocksSize += RequestedBlockSize;
+		BlockFreeFlag[BlockIndex] = false;
+		FreeBlockIndex++;
+
+		BlockHandle = BlockIndex;
+		BlockMemory = RawMemory + BlockOffets[BlockIndex];
+	}
+
+	return TPair<FParamStorageHandle, uint8*>(BlockHandle, BlockMemory);
+}
+
+void FParamStorage::ReleaseBlock(FParamStorageHandle BlockHandle)
+{
+	int32 BlockIndex = BlockHandle;
+
+	check(BlockIndex >= 0 && BlockIndex < FreeBlockIndex);
+
+	// just remove the offset here
+	BlockFreeFlag[BlockIndex] = true;
+	
+	// if it is the last block, release the memory and deal with the prev blocks flagged as free, if any
+	if (BlockIndex == FreeBlockIndex - 1)
+	{
+		// Recover free slots
+		while (BlockIndex >= 0 && BlockFreeFlag[BlockIndex] == true)
+		{
+			BlockOffets[BlockIndex] = 0;
+			CurrentBlocksSize -= BlockSizes[BlockIndex];
+			BlockSizes[BlockIndex] = 0;
+			BlockIndex--;
+			FreeBlockIndex--;
+		}
+	}
+}
+
+uint8* FParamStorage::GetBlockMemory(FParamStorageHandle BlockHandle)
+{
+	uint8* BlockMemory = nullptr;
+
+	const int32 BlockIndex = BlockHandle;
+
+	if (ensure(BlockIndex >= 0 && BlockIndex < FreeBlockIndex))
+	{
+		ensure(BlockFreeFlag[BlockIndex] == false);
+		BlockMemory = (BlockFreeFlag[BlockIndex] == false) ? RawMemory + BlockOffets[BlockIndex] : nullptr;
+	}
+
+	return BlockMemory;
+}
+
+bool FParamStorage::CheckBufferSize(int32 RequestedBlockSize)
+{
+	bool BufferSizeCorrect = true;
+
+	if (CurrentBlocksSize + RequestedBlockSize > RawMemorySize)
+	{
+		if (AllowGrowing)
+		{
+			// TODO : Grow
+			check(false);
+		}
+		else
+		{
+			BufferSizeCorrect = false;
+		}
+	}
+
+	return BufferSizeCorrect;
+}
+
+bool FParamStorage::CheckNumBlocks(int32 BlockIndex)
+{
+	bool ValidIndex = true;
+
+	if (BlockIndex >= BlockSizes.Num())
+	{
+		if (AllowGrowing)
+		{
+			// TODO : Grow
+		}
+		else
+		{
+			ValidIndex = false;
+		}
+		// TODO : Grow
+	}
+
+	return ValidIndex;
+}
+
+} // end namespace UE::DataInterface
diff --git a/Engine/Plugins/Experimental/DataInterface/Source/DataInterface/Private/DataInterfaceState.cpp b/Engine/Plugins/Experimental/DataInterface/Source/DataInterface/Private/DataInterfaceState.cpp
--- a/Engine/Plugins/Experimental/DataInterface/Source/DataInterface/Private/DataInterfaceState.cpp
+++ b/Engine/Plugins/Experimental/DataInterface/Source/DataInterface/Private/DataInterfaceState.cpp
@@ -42,10 +42,10 @@ FParam* FState::AllocateState(const FInterfaceKeyWithIdAndStack& InKey, const FC
 	{
 	case EStatePersistence::Relevancy:
 		// TODO: chunked allocator for relevancy-based stuff?
-		return &RelevancyValueMap.Add(InKey, FRelevancyParam(InType, Data, InContext.GetChunkFlags(), InContext.UpdateCounter));
+		return &RelevancyValueMap.Add(InKey, FRelevancyParam(InType, Data, InContext.GetBatchedFlags(), InContext.UpdateCounter));
 		break;
 	case EStatePersistence::Permanent:
-		return &PermanentValueMap.Add(InKey, FParam(InType, Data, InContext.GetChunkFlags()));
+		return &PermanentValueMap.Add(InKey, FParam(InType, Data, InContext.GetBatchedFlags()));
 		break;
 	default:
 		check(false);
diff --git a/Engine/Plugins/Experimental/DataInterface/Source/DataInterface/Public/DataInterface.h b/Engine/Plugins/Experimental/DataInterface/Source/DataInterface/Public/DataInterface.h
--- a/Engine/Plugins/Experimental/DataInterface/Source/DataInterface/Public/DataInterface.h
+++ b/Engine/Plugins/Experimental/DataInterface/Source/DataInterface/Public/DataInterface.h
@@ -8,6 +8,52 @@
 namespace UE::DataInterface
 {
 
+// Call the interface with the provided context (which should contain the parameters and result for the interface)
+template<typename InterfaceType>
+static bool GetDataSafe(const TScriptInterface<InterfaceType> DataInterface, const FContext& InContext)
+{
+	if(DataInterface.GetInterface() != nullptr)
+	{
+		return DataInterface.GetInterface()->GetData(InContext);
+	}
+
+	return false;
+}
+
+//// Set the Result into the context and call the interface
+//template<typename InterfaceType, typename ValueType>
+//static bool GetDataSafe(const TScriptInterface<InterfaceType> DataInterface, FContext& InContext, ValueType& Result)
+//{
+//	if (DataInterface.GetInterface() != nullptr)
+//	{
+//		if constexpr (TIsDerivedFrom<ValueType, FParam>::Value)
+//		{
+//			const FContext ResultHolder = InContext.WithResult(Result);
+//			return DataInterface.GetInterface()->GetData(ResultHolder);
+//		}
+//		else
+//		{
+//			// Support containers with Num() and GetData()
+//			if constexpr (TModels<Private::CSizedContainerWithAccessibleDataAsRawPtr, ValueType>::Value)
+//			{
+//				TWrapParam<ValueType::ElementType> WrappedResult(Result.GetData(), Result.Num());
+//				const FContext ResultHolder = InContext.WithResult(WrappedResult);
+//				return DataInterface.GetInterface()->GetData(ResultHolder);
+//			}
+//			else
+//			{
+//				TWrapParam<ValueType> WrappedResult(&Result);
+//				const FContext ResultHolder = InContext.WithResult(WrappedResult);
+//				return DataInterface.GetInterface()->GetData(ResultHolder);
+//			}
+//		}
+//	}
+//
+//	return false;
+//}
+
+// This version accepts a const Context, but it crates a sub context in order to be able to pass the Result
+// Using the non const version is a better option if possible
 template<typename InterfaceType, typename ValueType>
 static bool GetDataSafe(const TScriptInterface<InterfaceType> DataInterface, const FContext& InContext, ValueType& Result)
 {
@@ -19,8 +65,17 @@ static bool GetDataSafe(const TScriptInterface<InterfaceType> DataInterface, con
 		}
 		else
 		{
-			TWrapParam<ValueType> WrappedResult(InContext, &Result);
-			return DataInterface.GetInterface()->GetData(InContext, WrappedResult);
+			// Support containers with Num() and GetData()
+			if constexpr (TModels<Private::CSizedContainerWithAccessibleDataAsRawPtr, ValueType>::Value)
+			{
+				TWrapParam<typename ValueType::ElementType> WrappedResult(Result.GetData(), Result.Num());
+				return DataInterface.GetInterface()->GetData(InContext, WrappedResult);
+			}
+			else
+			{
+				TWrapParam<ValueType> WrappedResult(&Result);
+				return DataInterface.GetInterface()->GetData(InContext, WrappedResult);
+			}
 		}
 	}
 
@@ -36,7 +91,7 @@ static bool GetParameter(const FContext& InContext, FName InKey, ValueType& Resu
 	}
 	else
 	{
-		TWrapParam<ValueType> WrappedResult(InContext, &Result);
+		TWrapParam<ValueType> WrappedResult(&Result);
 		return InContext.GetParameter(InKey, WrappedResult);
 	}
 }
diff --git a/Engine/Plugins/Experimental/DataInterface/Source/DataInterface/Public/DataInterfaceContext.h b/Engine/Plugins/Experimental/DataInterface/Source/DataInterface/Public/DataInterfaceContext.h
--- a/Engine/Plugins/Experimental/DataInterface/Source/DataInterface/Public/DataInterfaceContext.h
+++ b/Engine/Plugins/Experimental/DataInterface/Source/DataInterface/Public/DataInterfaceContext.h
@@ -6,6 +6,7 @@
 #include "UObject/Class.h"
 #include "DataInterfaceParam.h"
 #include "DataInterfaceState.h"
+#include "DataInterfaceParamStorage.h"
 #include "DataInterfaceKey.h"
 #include "Misc/MemStack.h"
 
@@ -37,99 +38,334 @@ private:
 	friend struct FState;
 
 public:
+	static constexpr int32 MAX_BATCH_ELEMENTS = 64;
+
 	// Root public constructor. Constructs a context given a state.
-	// Further usage must use WithResult etc. to allow writable results to be used or param
-	FContext(float InDeltaTime, FState& InState, IDataInterfaceParameters* InParameters = nullptr);
+	FContext(float InDeltaTime, FState& InState, FParamStorage& InParamStorage, IDataInterfaceParameters* InParameters = nullptr);
+	
+	FContext(const FContext& other) = delete;
+	FContext& operator=(const FContext&) = delete;
+
+	FContext(FContext&& Other)
+		: FContext()
+	{
+		Swap(*this, Other);
+	}
+
+	FContext& operator= (FContext&& Other)
+	{
+		Swap(*this, Other);
+		return *this;
+	}
+
+	~FContext();
+
+public:
+	// --- Sub Context Creation --- 
 
+	// Create a sub context from this one that includes the provided Result
 	FContext WithResult(FParam& InResult) const;
 
+	// Create a sub context from this one that includes the provided parameter
 	FContext WithParameter(FName ParameterId, const FParam& InParameter) const;
 
+	// Create a sub context from this one that includes the provided parameters
 	FContext WithParameters(TArrayView<const TPair<FName, FParam>> InParameters) const;
-	
+
+	// Create a sub context from this one that includes the provided Result and parameters
+	FContext WithResultAndParameters(FParam& InResult, TArrayView<const TPair<FName, FParam>> InParameters) const;
+
+	// Create a sub context from this one that includes the provided interface parameter
 	FContext WithParameters(IDataInterfaceParameters* InParameters) const;
-	
-	bool GetParameter(FName InKey, FParam& OutParam) const;
 
 public:
-	// Gets (and allocates, if necessary) state for the specified node given this calling context
-	template<typename ValueType, EStatePersistence Persistence = EStatePersistence::Relevancy>
-	TParam<ValueType> GetState(const FInterfaceKeyWithId& InKey) const
+
+	// --- Interface for direct Param storage (prototype) ---
+
+	FContext CreateSubContext() const;
+
+	enum class EParamType : uint8
+	{
+		None		= 0,
+		Input		= 1 << 0,
+		Output		= 1 << 1,
+	};
+
+	// Add a parameter by value, copying it to the shared storage
+	template<typename ValueType>
+	FHParam AddValue(EParamType InParamType, FName ParameterId, ValueType &Value)
 	{
-		return State.GetState<ValueType, Persistence>(InKey, *this, CallstackHash);
+		// Values are added mutable even if source is const
+		const FParam::EFlags Flags = FParam::EFlags::Value | FParam::EFlags::Mutable;
+
+		return AddParameter(InParamType, ParameterId, Flags, &Value);
 	}
 
-	// Gets (and allocates, if necessary) state for the specified node given this calling context
-	template<typename ValueType, EStatePersistence Persistence = EStatePersistence::Relevancy>
-	TParam<ValueType> GetState(const IDataInterface* InDataInterface, uint32 InId) const
+	// Add a parameter by value, copying it to the shared storage
+	template<typename ValueType>
+	FHParam AddValue(EParamType InParamType, FName ParameterId, ValueType&& Value)
 	{
-		return State.GetState<ValueType, Persistence>(InDataInterface, InId, *this, CallstackHash);
+		// Values are added mutable even if source is const
+		const FParam::EFlags Flags = FParam::EFlags::Value | FParam::EFlags::Mutable;
+
+		return AddParameter(InParamType, ParameterId, Flags, &Value);
 	}
-	
-public:
-	// Set a result directly
+
+	// Add a parameter by value, copying it to the shared storage
 	template<typename ValueType>
-	void SetResult(const ValueType& InValue) const
+	FHParam AddValue(EParamType InParamType, FName ParameterId, ValueType* Value)
+	{
+		// Values are added mutable even if source is const
+		const FParam::EFlags Flags = FParam::EFlags::Value | FParam::EFlags::Mutable;
+
+		return AddParameter(InParamType, ParameterId, Flags, Value);
+	}
+
+	// Add a parameter by reference, adding just a pointer to the shared storage
+	template<typename ValueType>
+	FHParam AddReference(EParamType InParamType, FName ParameterId, ValueType& Value)
 	{
-		if(Result != nullptr)
+		const FParam::EFlags Flags = FParam::EFlags::Reference;
+
+		return AddParameter(InParamType, ParameterId, Flags, &Value);
+	}
+
+	// Add a parameter by reference, adding just a pointer to the shared storage
+	template<typename ValueType>
+	FHParam AddReference(EParamType InParamType, FName ParameterId, ValueType *Value)
+	{
+		FParam::EFlags Flags = FParam::EFlags::Reference;
+
+		return AddParameter(InParamType, ParameterId, Flags, Value);
+	}
+
+	template<typename ValueType>
+	FHParam AddParameter(EParamType InParamType, FName ParameterId, FParam::EFlags Flags, ValueType* Value)
+	{
+		FHParam ParamHandle;
+
+		check(InParamType == EParamType::Input || InParamType == EParamType::Output);
+		check(EnumHasAnyFlags(Flags, FParam::EFlags::Value) || EnumHasAnyFlags(Flags, FParam::EFlags::Reference));
+
+		// check if it is a simple FHParam copy
+		if constexpr (TIsDerivedFrom<ValueType, FHParam>::Value)
+		{
+			const FParam* ExistingParam = ParamStorage->GetParam(Value->ParamHandle);
+			check(ExistingParam != nullptr);
+
+			const bool bIsReadOnly = EnumHasAnyFlags(ExistingParam->Flags, FParam::EFlags::Mutable) == false;
+
+			// Verify that we want to create an Input param or if we want an output, the source is not read only
+			check(InParamType == EParamType::Input || bIsReadOnly == false);
+
+			// For now create a copy of the Param (TODO : change ref to value or value to ref, if the caller requests it?)
+			ParamHandle = *Value;
+		}
+		else
 		{
-			// Support containers with Num() and GetData()
+			using MutableValueType = typename TRemoveConst<ValueType>::Type;
+
+			// If non const set the mutable flag
+			if constexpr (TIsConst<ValueType>::Value == false)
+			{
+				EnumAddFlags(Flags, FParam::EFlags::Mutable);
+			}
+
+			// deal with containers
 			if constexpr (TModels<Private::CSizedContainerWithAccessibleDataAsRawPtr, ValueType>::Value)
 			{
-				checkSlow(Private::CheckParam<ValueType>(Result));
-				checkSlow(InValue.Num() == State.ChunkSize);
-
-				TArrayView<ValueType> Results(static_cast<ValueType*>(Result->Data), State.NumElements);
-				const ValueType* Data = InValue.GetData();
-				for(int32 ResultIndex = 0; ResultIndex < State.NumElements; ++ResultIndex)
-				{
-					Results[ResultIndex] = Data[ResultIndex];
-				}
+				MutableValueType* ParamData = const_cast<MutableValueType>(static_cast<ValueType*>(Value->GetData()));
+				ParamHandle = EnumHasAnyFlags(Flags, FParam::EFlags::Value) 
+					? ParamHandle = ParamStorage->AddValue(ParamData, Value->Num(), Flags)
+					: ParamStorage->AddReference(ParamData, Value->Num(), Flags);
 			}
+			// and finally with single params
 			else
 			{
-				checkSlow(Private::CheckParam<ValueType>(Result));
+				ParamHandle = EnumHasAnyFlags(Flags, FParam::EFlags::Value)
+					? ParamHandle = ParamStorage->AddValue(const_cast<MutableValueType*>(Value), 1, Flags)
+					: ParamStorage->AddReference(const_cast<MutableValueType*>(Value), 1, Flags);
+			}
+		}
+
+		AdditionalParameterHandles.Add(ParameterId, ParamHandle);
+
+		return ParamHandle;
+	}
+
+	// Get a TParam from a handle
+	template<typename ValueType>
+	TParam<ValueType> GetParameterChecked(const FHParam &InParamHandle) const
+	{
+		const FParam* Param = ParamStorage->GetParam(InParamHandle.ParamHandle);
+
+		TParam<ValueType> RetVal(Param, Param->GetFlags());
+
+		check(RetVal.IsValid());							// found something
+		checkSlow(Private::CheckParam<ValueType>(RetVal));	// the type found is valid for what was requested
+
+		return RetVal;
+	}
+
+	// Get a parameter from a handle as a specified type
+	template <typename ValueType>
+	ValueType& GetParameterAs(const FHParam& InParamHandle, int32* OutNumElems = nullptr)
+	{
+		const FParam* Param = ParamStorage->GetParam(InParamHandle.ParamHandle);
+
+		check(Param != nullptr);
+		check(Param->GetType().GetTypeId() == Private::TParamType<ValueType>::GetType().GetTypeId());
+
+		// If a non cost is requested, check the param has mutable flag
+		if constexpr (TIsConst<ValueType>::Value == false)
+		{
+			check(Param->IsMutable());
+		}
+
+		ValueType* ValueData = EnumHasAnyFlags(Param->Flags, FParam::EFlags::Embedded)
+			? static_cast<ValueType*>((void*)&Param->Data)
+			: static_cast<ValueType*>(Param->Data);
+
+		if (OutNumElems != nullptr)
+		{
+			*OutNumElems = Param->GetNumElements();
+		}
+
+		return *ValueData;
+	}
+
+	// --- Parameter management --- 
+
+	// Get a parameter if exist, returns false if not
+	bool GetParameter(FName InKey, FParam& OutParam) const;
+
+	// Get a parameter as a specified type, checking it exist and the type
+	template<typename ValueType>
+	TParam<ValueType> GetParameterChecked(FName InKey) const
+	{
+		TParam<ValueType> RetVal(GetBatchedFlags());
+
+		GetParameter(InKey, RetVal);
+
+		check(RetVal.IsValid());							// found something
+		checkSlow(Private::CheckParam<ValueType>(RetVal));	// the type found is valid for what was requested
+
+		return RetVal;
+	}
+
+	// Get a parameter as a specified type, returning a TOptional
+	template<typename ValueType>
+	TOptional<TParam<ValueType>> GetParameter(FName InKey) const
+	{
+		TOptional<TParam<ValueType>> RetVal;
+
+		TParam<ValueType> Param(FParam::EFlags::Mutable | GetBatchedFlags());
+		if (GetParameter(InKey, Param))
+		{
+			checkSlow(Private::CheckParam<ValueType>(Param));
+
+			RetVal = Param;
+		}
+
+		return RetVal;
+	}
+
+public:
+	// --- Result Management ---
+
+	// Get the ResultParam, checking it has been set
+	FParam& GetResultParam()
+	{
+		check(Result != nullptr);
+		return *Result;
+	}
+
+	// Get the ResultParam, checking it has been set
+	FParam& GetResultParam() const
+	{
+		check(Result != nullptr);
+		return *Result;
+	}
+
+	// Set a HParam as a result (for compatibility reasons)
+	void SetHParamAsResult(const FHParam& InHParam)
+	{
+		Result = ParamStorage->GetParam(InHParam.ParamHandle);
+	}
+
+	// Set a result value directly
+	// The receiver Result Param must be set prior to this call
+	template<typename ValueType>
+	void SetResult(const ValueType& InValue) const
+	{
+		// Support containers with Num() and GetData()
+		if constexpr (TModels<Private::CSizedContainerWithAccessibleDataAsRawPtr, ValueType>::Value)
+		{
+			checkSlow(Private::CheckParam<ValueType>(*Result));
+			checkSlow(InValue.Num() == Result->NumElements);
+
+			TArrayView<ValueType> Results(*static_cast<ValueType*>(Result->Data), Result->NumElements);
+			const ValueType* Data = InValue.GetData();
+			for(int32 ResultIndex = 0; ResultIndex < Result->NumElements; ++ResultIndex)
+			{
+				Results[ResultIndex] = Data[ResultIndex];
+			}
+		}
+		else
+		{
+			checkSlow(Private::CheckParam<ValueType>(*Result));
+			checkSlow(Result->NumElements == 1);
 
-				TArrayView<ValueType> Results(static_cast<ValueType*>(Result->Data), State.NumElements);
-				for(ValueType& ResultElement : Results)
-				{
-					ResultElement = InValue;
-				}
+			TArrayView<ValueType> Results(static_cast<ValueType*>(Result->Data), Result->NumElements);
+			for(ValueType& ResultElement : Results)
+			{
+				ResultElement = InValue;
 			}
 		}
 	}
-	
-	// Get the current result as a specified type
+
+	// Get the current result as a TParam
 	template<typename ValueType>
 	TParam<ValueType> GetResult() const
 	{
-		// Cannot call this without a result being present
-		check(Result);
-		
 		// Types must match
 		checkSlow(Private::CheckParam<ValueType>(*Result));
-		return TParam<ValueType>(*Result, GetChunkFlags());
+		return TParam<ValueType>(*Result, GetBatchedFlags());
 	}
-	
-	// Get the current result as a raw parameter
-	const FParam& GetResultRaw() const
+
+	// Get the current result as a specified raw type pointer
+	template<typename ValueType>
+	ValueType* GetResultChecked() const
 	{
-		// Cannot call this without a result being present
-		check(Result);
-		return *Result;
+		// Types must match
+		checkSlow(Private::CheckParam<ValueType>(*Result));
+		return static_cast<ValueType*>(Result->Data);
 	}
 
 public:
-	// Number of elements that this context represents
-	uint32 GetNum() const
+	// --- State management ---
+
+	// Gets (and allocates, if necessary) state for the specified node given this calling context
+	template<typename ValueType, EStatePersistence Persistence = EStatePersistence::Relevancy>
+	TParam<ValueType> GetState(const FInterfaceKeyWithId& InKey) const
+	{
+		return State->GetState<ValueType, Persistence>(InKey, *this, CallstackHash);
+	}
+
+	// Gets (and allocates, if necessary) state for the specified node given this calling context
+	template<typename ValueType, EStatePersistence Persistence = EStatePersistence::Relevancy>
+	TParam<ValueType> GetState(const IDataInterface* InDataInterface, uint32 InId) const
 	{
-		return State.NumElements;
+		return State->GetState<ValueType, Persistence>(InDataInterface, InId, *this, CallstackHash);
 	}
 
-	// Max number of elements that this context represents
-	uint32 GetChunkSize() const
+public:
+	// --- Mix Context Utils ---
+
+	// Number of elements that this context represents
+	int32 GetNum() const
 	{
-		return State.ChunkSize;
+		return State->NumElements;
 	}
 
 	// Access delta time as a param
@@ -137,108 +373,120 @@ public:
 
 	// Raw access to delta time
 	float GetDeltaTime() const { return DeltaTime; }
-	
-	// Get parameter flags for chunking
-	FParam::EFlags GetChunkFlags() const
+
+	// --- Batching management ---
+
+	// Get parameter flags for batching
+	FParam::EFlags GetBatchedFlags() const
+	{
+		return State->NumElements > 1 ? FParam::EFlags::Batched : FParam::EFlags::None;
+	}
+
+	FORCEINLINE_DEBUGGABLE bool ShouldProcessThisElement(int32 ElementIndex) const
+	{
+		check(ElementIndex < MAX_BATCH_ELEMENTS);
+		return (BranchingMask & (1ull << ElementIndex)) != 0;
+	}
+
+	FORCEINLINE_DEBUGGABLE const uint64 GetBranchingMask() const
+	{
+		return BranchingMask;
+	}
+
+	void SetBranchingMask(const TArrayView<int32> InEnabledElementIndexes)
+	{
+		check(InEnabledElementIndexes.Num() <= GetNum());
+
+		BranchingMask = 0;
+		for (int32 ElementIndex : InEnabledElementIndexes)
+		{
+			check(ElementIndex < MAX_BATCH_ELEMENTS);
+
+			BranchingMask |= (1ull << ElementIndex);
+		}
+	}
+
+	void SetBranchingMask(uint64 InBranchingMask)
 	{
-		return State.ChunkSize > 1 ? FParam::EFlags::Chunked : FParam::EFlags::None;
+		BranchingMask = InBranchingMask;
 	}
 
 private:
-	FContext(float InDeltaTime, FState& InState, FParam* InResult)
-		: State(InState)
-		, Result(InResult)
+	FContext(float InDeltaTime, FState& InState, FParamStorage& InParamStorage, FParam& InResult)
+		: State(&InState)
+		, ParamStorage(&InParamStorage)
+		, Result(&InResult)
 		, DeltaTime(InDeltaTime)
+		, BranchingMask(MAX_uint64)
 	{}
 
+	FContext() = default;
+
 	FContext WithCallRaw(const IDataInterface* InDataInterface) const;
 
 	void FlushRelevancy() const;
 
+	int32 GetParametersSize(TArrayView<const TPair<FName, FParam>> InParameters, TArray<int32>& ParamAllocSizes) const;
+	void AddParameters(TArrayView<const TPair<FName, FParam>> InParameters);
+
 private:
 	TMap<FName, FParam> AdditionalParameters;
+	TMap<FName, FHParam> AdditionalParameterHandles;  // Temp param storage prototype
 	const FContext* Parent = nullptr;
 	const FContext* Root = nullptr;
-	FState& State;
+	FState* State = nullptr;
+	FParamStorage* ParamStorage = nullptr;
+	FParamStorageHandle BlockHandle = InvalidBlockHandle;
 	FParam* Result = nullptr;
 	IDataInterfaceParameters* Parameters = nullptr;
 	float DeltaTime = 0.0f;
 	uint32 CallstackHash = 0;
 	uint32 UpdateCounter = 0;
+	uint64 BranchingMask = MAX_uint64;
 };
 
-// A typed param that wraps a user ptr
-template<typename ValueType>
-struct TWrapParam : TParam<ValueType>
+// A typed param that owns it's own memory
+template<typename ValueType, typename AllocatorType = TMemStackAllocator<>>
+struct TAllocParam : public TParam<ValueType>
 {
 public:
-	TWrapParam(const FContext& InContext, ValueType* InValuePtrToWrap)
-		: TParam<ValueType>(GetDataForConstructor(InContext, InValuePtrToWrap), GetFlagsForConstructor(InContext, InValuePtrToWrap))
+	TAllocParam(const FContext& InContext)
+		: TParam<ValueType>(InContext.GetBatchedFlags())
 	{
+		TParam<ValueType>::Data = GetDataForConstructor(InContext);
 	}
 
-private:
-	void* GetDataForConstructor(const FContext& InContext, ValueType* InValuePtrToWrap) const
+	TAllocParam(const FContext& InContext, FParam::EFlags InFlags)
+		: TParam<ValueType>(InFlags)
 	{
-		using MutableValueType = typename TRemoveConst<ValueType>::Type;
-		
-		// Support containers with Num() and GetData()
-		if constexpr (TModels<Private::CSizedContainerWithAccessibleDataAsRawPtr, ValueType>::Value)
-		{
-			checkSlow(InContext.GetChunkSize() == InValuePtrToWrap->Num());
-			return const_cast<void*>(InValuePtrToWrap->GetData());
-		}
-		else
-		{
-			checkSlow(InContext.GetChunkSize() == 1);
-			return const_cast<MutableValueType*>(InValuePtrToWrap);
-		}
+		TParam<ValueType>::Data = GetDataForConstructor(InContext);
 	}
 
-	FParam::EFlags GetFlagsForConstructor(const FContext& InContext, ValueType* InValuePtrToWrap) const
-	{
-		FParam::EFlags NewFlags = FParam::EFlags::None;
-		
-		// Add chunked flags or not
-		if constexpr (TModels<Private::CSizedContainerWithAccessibleDataAsRawPtr, ValueType>::Value)
-		{
-			// Containers always assume chunked
-			checkSlow(InContext.GetChunkSize() == InValuePtrToWrap->Num());
-			NewFlags |= InContext.GetChunkFlags();
-		}
-		else
-		{
-			checkSlow(InContext.GetChunkSize() == 1);
-		}
-
-		// Add const flags or not
-		if constexpr (!TIsConst<ValueType>::Value)
-		{
-			NewFlags |= FParam::EFlags::Mutable;
-		}
-
-		return NewFlags;
-	}
-};
-
-// A typed param that owns it's own memory
-template<typename ValueType, typename AllocatorType = TMemStackAllocator<>>
-struct TAllocParam : TParam<ValueType>
-{
-public:
-	TAllocParam(const FContext& InContext)
-		: TParam<ValueType>(GetDataForConstructor(InContext), InContext.GetChunkFlags())
-	{
-	}
 
 private:
 	void* GetDataForConstructor(const FContext& InContext)
 	{
-		ValueArray.SetNum(InContext.GetChunkSize());
+		ValueArray.SetNum(InContext.GetNum());
 		return ValueArray.GetData();
 	}
 
 	TArray<ValueType, AllocatorType> ValueArray;
 };
 
-}
\ No newline at end of file
+
+// --- Batching helper macros ---
+
+#define SKIP_CONTEXT_ELEMENT_IF_MASKED(Context) \
+		if (!Context.ShouldProcessThisElement(i)) \
+			continue;
+
+// Helper macros to allow processing context element batches
+#define PROCESS_BATCH_ELEMENTS_START(Context)  \
+	const int32 ContextNumElements = Context.GetNum(); \
+	for (int32 i = 0; i < ContextNumElements; i++) \
+	{ \
+		SKIP_CONTEXT_ELEMENT_IF_MASKED(Context)
+
+#define PROCESS_BATCH_ELEMENTS_END }
+
+}
diff --git a/Engine/Plugins/Experimental/DataInterface/Source/DataInterface/Public/DataInterfaceKernel.h b/Engine/Plugins/Experimental/DataInterface/Source/DataInterface/Public/DataInterfaceKernel.h
--- a/Engine/Plugins/Experimental/DataInterface/Source/DataInterface/Public/DataInterfaceKernel.h
+++ b/Engine/Plugins/Experimental/DataInterface/Source/DataInterface/Public/DataInterfaceKernel.h
@@ -11,38 +11,25 @@ namespace UE::DataInterface
 struct DATAINTERFACE_API FKernel
 {
 public:
-	// Runs a 'kernel' predicate that iterates over the supplied (optionally chunked) params
+	// Runs a 'kernel' predicate that iterates over the supplied (optionally batched) params
 	// Number and type of params must match between the predicate & param pack
+	// The kernel supports passing a single element parameter where N are expected (as a compacted constant)
 	template<typename PredicateType, typename... ParamTypes>
 	static void Run(const FContext& InContext, PredicateType InPredicate, ParamTypes&... InArgs)
 	{
-		// Cache offsets to apply each element for each parameter
-		const uint32 Offsets[sizeof...(InArgs)] =
-		{
-			[](auto& InArg) -> uint32
-			{
-				if constexpr (TIsDerivedFrom<ParamTypes, FParam>::Value)
-				{
-					if(EnumHasAnyFlags(InArg.Flags, FParam::EFlags::Chunked))
-					{
-						return InArg.GetType().GetSize();
-					}
-				}
-
-				return 0;
-			}(InArgs)...
-		};
-
 		// Iterate per-element in the context, applying kernel predicate to all elements
 		const uint32 NumElements = InContext.GetNum();
-		for(uint32 ElementIndex = 0; ElementIndex < NumElements; ++ElementIndex)
+		for (uint32 ElementIndex = 0; ElementIndex < NumElements; ++ElementIndex)
 		{
+			if (InContext.ShouldProcessThisElement(ElementIndex) == false)
+				continue;
+
 			uint32 ParameterIndex = 0;
 
 			// Call predicate, expanding InArgs parameter pack from TParam<ParamType> -> ParamType
-			// and applying correct offset for each param
+			// and applying correct offset for each param (using pointer arithmetic)
 			InPredicate(
-				[&Offsets, &ParameterIndex](uint32 InElementIndex, auto& InArg) -> auto&
+				[&ParameterIndex](uint32 InElementIndex, auto& InArg) -> auto&
 				{
 					auto GetBasePtr = [](auto& InArg)
 					{
@@ -57,16 +44,33 @@ public:
 						}
 					};
 
-					auto GetOffset = [&Offsets, &ParameterIndex](uint32 InElementIndex) -> uint32
+					auto GetOffset = [&ParameterIndex](auto& InArg, int32 InElementIndex) -> int32
 					{
-						return Offsets[ParameterIndex++] * InElementIndex;
+						if constexpr (TIsDerivedFrom<ParamTypes, FParam>::Value)
+						{
+							if (InArg.NumElements > 1)
+							{
+								check(InElementIndex < InArg.NumElements);
+								return InElementIndex;
+							}
+							else // single elements always use 0 offset, even if used on N element batch
+							{
+								check(InArg.NumElements > 0);
+								return 0;
+							}
+						}
+						else
+						{
+							return 0;
+						}
 					};
-		
-					return *(GetBasePtr(InArg) + GetOffset(InElementIndex));
+
+					// note that this uses pointer arithmetic, so an offset of +1 means adding the sizeof of the param type
+					return *(GetBasePtr(InArg) + GetOffset(InArg, InElementIndex));
 				}(ElementIndex, InArgs)...
 			);
 		}
 	}
 };
 
-}
\ No newline at end of file
+} // end namespace UE::DataInterface
\ No newline at end of file
diff --git a/Engine/Plugins/Experimental/DataInterface/Source/DataInterface/Public/DataInterfaceParam.h b/Engine/Plugins/Experimental/DataInterface/Source/DataInterface/Public/DataInterfaceParam.h
--- a/Engine/Plugins/Experimental/DataInterface/Source/DataInterface/Public/DataInterfaceParam.h
+++ b/Engine/Plugins/Experimental/DataInterface/Source/DataInterface/Public/DataInterfaceParam.h
@@ -10,6 +10,7 @@ namespace UE::DataInterface
 {
 
 struct FParamType;
+struct FParamStorage;
 
 // Parameter/result/state memory wrapper
 // @TODO can we save memory here with some base/derived system for fundamental types vs struct types?
@@ -18,45 +19,71 @@ struct DATAINTERFACE_API FParam
 public:
 	enum class EFlags : uint8
 	{
-		None	= 0,		// No flags
-		Mutable	= 1,		// Parameter is mutable, so can be mutated at runtime
-		Chunked = 2,		// Parameter is chunked, so represents a number of values rather than a single value. Chunk size is determined by an associated FContext
+		None		= 0,			// No flags
+		Mutable		= 1 << 0,		// Parameter is mutable, so can be mutated at runtime
+		Batched		= 1 << 1,		// Parameter is batched, so represents a number of values rather than a single value. Batch size is determined by NumElements and associated FContext
+		Stored		= 1 << 2,		// Parameter has to be stored on context storage
+		
+		// --- Flags added for parameter storage prototype ---
+		Value		= 1 << 3,		// Parameter will be stored as a Value
+		Reference	= 1 << 4,		// Parameter will be stored as a Reference (pointer)
+		Embedded	= 1 << 5,		// Parameter will be stored as a Value, but stored directly on the Data pointer
 	};
 
 	// Get the type of this param 
 	const FParamType& GetType() const;
 
 	// Check whether the supplied param can be written to by this param
-	// Verifies that types are compatible, destination is mutable and chunking matches
+	// Verifies that types are compatible, destination is mutable and batching matches
 	bool CanAssignTo(const FParam& InParam) const;
 
 	// Helper function for CanAssignTo
-	bool CanAssignWith(const FParamType& InType, EFlags InFlags, FStringBuilderBase* OutReasonPtr = nullptr) const;
+	bool CanAssignWith(const FParamType& InType, EFlags InFlags, int32 InNumelements, FStringBuilderBase* OutReasonPtr = nullptr) const;
+
+	bool IsMutable() const { return EnumHasAnyFlags(Flags, EFlags::Mutable); }
+
+	bool IsBatched() const { return EnumHasAnyFlags(Flags, EFlags::Batched); }
 
-	bool bIsMutable() const { return EnumHasAnyFlags(Flags, EFlags::Mutable); }
+	int32 GetNumElements() const { return NumElements; }
 
-	bool bIsChunked() const { return EnumHasAnyFlags(Flags, EFlags::Chunked); }
+	EFlags GetFlags() const { return Flags; }
+	
+	void* GetData() const { return Data; }
 
 protected:
 	friend struct FContext;
 	friend struct FState;
 	friend struct FKernel;
-
-	FParam() = default;
+	friend struct FParamStorage;
 	
 	FParam(const FParam* InOtherParam);
 	FParam(const FParamType& InType, void* InData, EFlags InFlags);
+	FParam(const FParamType& InType, void* InData, int32 InNumElements, EFlags InFlags);
+	FParam(const FParamType& InType, EFlags InFlags);
 
 	void* Data = nullptr;
+	int32 NumElements = 0;
 	int16 TypeId = INDEX_NONE;
 	EFlags Flags = EFlags::None;
+
+public:
+	// Internal use, but required for default constructed elements on containers
+	FParam() = default;
+
+	static FParam MakeParam_Internal(const FParamType& InType, void* InData, int32 InNumElements, EFlags InFlags)
+	{
+		return FParam(InType, InData, InNumElements, InFlags);
+	}
 };
 
 ENUM_CLASS_FLAGS(FParam::EFlags);
 
+
 // Simple type representation for data interfaces
 struct DATAINTERFACE_API FParamType
 {
+	typedef FParam(*ParamCopyFunction)(const FParam& InSource, void* InMemory, int32 InAllocatedMemory);
+
 	FParamType() = default;
 
 	inline const UScriptStruct* GetStruct() const
@@ -84,6 +111,11 @@ struct DATAINTERFACE_API FParamType
 		return TypeId;
 	}
 
+	ParamCopyFunction GetParamCopyFunction() const
+	{
+		return ParamCopyFunctionPtr;
+	}
+
 	struct DATAINTERFACE_API FRegistrar
 	{
 		// Helper for registration (via IMPLEMENT_DATA_INTERFACE_PARAM_TYPE).
@@ -91,7 +123,7 @@ struct DATAINTERFACE_API FParamType
 		FRegistrar(TUniqueFunction<void(void)>&& InFunction);
 
 		// Helper for registration (via IMPLEMENT_DATA_INTERFACE_PARAM_TYPE)
-		static void RegisterType(FParamType& InType, const UScriptStruct* InStruct, FName InName, uint32 InSize, uint32 InAlignment);
+		static void RegisterType(FParamType& InType, const UScriptStruct* InStruct, FName InName, uint32 InSize, uint32 InAlignment, const ParamCopyFunction InParamCopyFunctionPtr);
 
 		// Helper for registration
 		static void RegisterDeferredTypes();
@@ -109,6 +141,7 @@ private:
 	uint32 Size = 0;
 	uint16 Alignment = 0;
 	uint16 TypeId = INDEX_NONE;
+	ParamCopyFunction ParamCopyFunctionPtr = nullptr;
 };
 
 // A type-erased param that wraps FProperty memory
@@ -168,13 +201,93 @@ struct CSizedContainerWithAccessibleDataAsRawPtr
 template<typename ValueType>
 struct TParam : FParam
 {
+public:
+	FORCEINLINE_DEBUGGABLE bool IsValid() const
+	{
+		return ((Data != nullptr || EnumHasAnyFlags(Flags, FParam::EFlags::Embedded))
+			&& TypeId != INDEX_NONE 
+			&& NumElements > 0);
+	}
+
+	FORCEINLINE_DEBUGGABLE const ValueType& operator[](int32 InIndex) const
+	{
+		check(InIndex >= 0 && (InIndex < NumElements || NumElements == 1));
+		check(Data != nullptr || EnumHasAnyFlags(Flags, FParam::EFlags::Embedded));
+
+		const ValueType* ValueData = EnumHasAnyFlags(Flags, FParam::EFlags::Embedded)
+			? static_cast<const ValueType*>((void*)&Data)
+			: static_cast<const ValueType*>(Data);
+		return (NumElements > InIndex) ? ValueData[InIndex] : ValueData[0];
+	}
+
+	FORCEINLINE_DEBUGGABLE ValueType& operator[](int32 InIndex)
+	{
+		check(InIndex >= 0 && (InIndex < NumElements || NumElements == 1));
+		check(Data != nullptr || EnumHasAnyFlags(Flags, FParam::EFlags::Embedded));
+
+		ValueType* ValueData = EnumHasAnyFlags(Flags, FParam::EFlags::Embedded)
+			? static_cast<ValueType*>((void*) &Data)
+			: static_cast<ValueType*>(Data);
+		return (NumElements > InIndex) ? ValueData[InIndex] : ValueData[0];
+	}
+
+	FORCEINLINE_DEBUGGABLE operator const ValueType& () const
+	{
+		check((Data != nullptr || EnumHasAnyFlags(Flags, FParam::EFlags::Embedded)) && NumElements == 1);
+
+		const ValueType* ValueData = EnumHasAnyFlags(Flags, FParam::EFlags::Embedded)
+			? static_cast<const ValueType*>((void*)&Data)
+			: static_cast<const ValueType*>(Data);
+
+		return *ValueData;
+	}
+
+	FORCEINLINE_DEBUGGABLE operator ValueType& ()
+	{
+		check((Data != nullptr || EnumHasAnyFlags(Flags, FParam::EFlags::Embedded)) && NumElements == 1);
+
+		ValueType* ValueData = EnumHasAnyFlags(Flags, FParam::EFlags::Embedded)
+			? static_cast<ValueType*>((void*)&Data)
+			: static_cast<ValueType*>(Data);
+
+		return *ValueData;
+	}
+
+	FORCEINLINE_DEBUGGABLE ValueType& operator*() const
+	{
+		check((Data != nullptr || EnumHasAnyFlags(Flags, FParam::EFlags::Embedded)) && NumElements == 1);
+
+		ValueType* ValueData = EnumHasAnyFlags(Flags, FParam::EFlags::Embedded)
+			? static_cast<ValueType*>((void*)&Data)
+			: static_cast<ValueType*>(Data);
+
+		return *ValueData;
+	}
+
+	FORCEINLINE_DEBUGGABLE ValueType* operator->() const
+	{
+		check((Data != nullptr || EnumHasAnyFlags(Flags, FParam::EFlags::Embedded)) && NumElements == 1);
+
+		ValueType* ValueData = EnumHasAnyFlags(Flags, FParam::EFlags::Embedded)
+			? static_cast<ValueType*>((void*)&Data)
+			: static_cast<ValueType*>(Data);
+
+		return ValueData;
+	}
+
 protected:
 	using TValueType = ValueType;
 	
 	friend struct FContext;
 	friend struct FState;
 	friend struct FKernel;
+	friend struct FParamStorage;
 	
+	TParam(EFlags InFlags)
+		: FParam(Private::TParamType<ValueType>::GetType(), InFlags)
+	{
+	}
+
 	TParam(void* InData, EFlags InFlags)
 		: FParam(Private::TParamType<ValueType>::GetType(), InData, InFlags)
 	{
@@ -185,16 +298,37 @@ protected:
 	{
 	}
 
+	TParam(void* InData, int32 InNumElements, EFlags InFlags)
+		: FParam(Private::TParamType<ValueType>::GetType(), InData, InNumElements, InFlags)
+	{
+	}
+
+	TParam(const void* InData, int32 InNumElements, EFlags InFlags)
+		: FParam(Private::TParamType<ValueType>::GetType(), const_cast<void*>(InData), InNumElements, InFlags)
+	{
+	}
+
 	TParam(FParam& InParam, EFlags InAdditionalFlags = EFlags::None)
 		: FParam(&InParam)
 	{
 #if DO_CHECK
 		TStringBuilder<64> Error;
-		checkf(InParam.CanAssignWith(Private::TParamType<ValueType>::GetType(), GetTypeFlags(InAdditionalFlags), &Error),
+		checkf(InParam.CanAssignWith(Private::TParamType<ValueType>::GetType(), GetTypeFlags(InAdditionalFlags), NumElements, &Error),
 			TEXT("Cannot assign type: %s"), Error.ToString());
 #endif
 	}
 
+	TParam(const FParam& InParam, EFlags InAdditionalFlags = EFlags::None)
+		: FParam(&InParam)
+	{
+#if DO_CHECK
+		TStringBuilder<64> Error;
+		checkf(InParam.CanAssignWith(Private::TParamType<ValueType>::GetType(), GetTypeFlags(InAdditionalFlags), NumElements, &Error),
+			TEXT("Cannot assign type: %s"), Error.ToString());
+#endif
+	}
+
+protected:
 	static EFlags GetTypeFlags(EFlags InAdditionalFlags)
 	{
 		EFlags Flags = InAdditionalFlags;
@@ -212,6 +346,37 @@ protected:
 namespace Private
 {
 
+template<typename ValueType>
+struct TParamCopyBaseTemplate
+{
+	static FParam CopyParam(const FParam& InSource, void* InTargetMemory, int32 InTargetAllocatedSize)
+	{
+		const FParamType& ParamType = InSource.GetType();
+		
+		const int32 NumElem = InSource.GetNumElements();
+		const UScriptStruct* ScriptStruct = ParamType.GetStruct();
+
+		if (ScriptStruct != nullptr)
+		{
+			const int32 StructSize = ScriptStruct->GetStructureSize();
+			check(NumElem * StructSize <= InTargetAllocatedSize);
+
+			ScriptStruct->CopyScriptStruct(InTargetMemory, InSource.GetData(), NumElem);
+		}
+		else
+		{
+			const int32 ParamAlignment = ParamType.GetAlignment();
+			const int32 ParamSize = ParamType.GetSize();
+			const int32 ParamAllocSize = Align(ParamSize, ParamAlignment);
+
+			// TODO : check if this is safe to do for any type not containing a script struct, else we can specialize for POD types
+			FMemory::Memcpy(InTargetMemory, InSource.GetData(), NumElem * ParamAllocSize);
+		}
+
+		return FParam::MakeParam_Internal(ParamType, InTargetMemory, NumElem, InSource.GetFlags());
+	}
+};
+
 struct CHasStaticStruct
 {
 	template<typename ValueType>
@@ -244,9 +409,186 @@ struct TParamTypeHelper
 	{
 		return alignof(ValueType);
 	}
+
+	FORCEINLINE static const FParamType::ParamCopyFunction GetParamCloneFunction()
+	{
+		return &TParamCopyBaseTemplate<ValueType>::CopyParam;
+	}
 };
 
 }
 
-}
+// A typed param that wraps a user ptr to single or batched elements
+template<typename ValueType>
+struct TWrapParam : TParam<ValueType>
+{
+public:
+	TWrapParam(ValueType* InValuePtrToWrap)
+		: TParam<ValueType>(GetDataForConstructor(InValuePtrToWrap), GetFlagsForConstructor(InValuePtrToWrap))
+	{
+	}
+
+	TWrapParam(ValueType& InValueToWrap)
+		: TParam<ValueType>(GetDataForConstructor(&InValueToWrap), GetFlagsForConstructor(&InValueToWrap))
+	{
+	}
+
+	TWrapParam(ValueType* InValuePtrToWrap, const int32 NumElements)
+		: TParam<ValueType>(InValuePtrToWrap, NumElements, NumElements > 1 ? FParam::EFlags::Batched : FParam::EFlags::None)
+	{
+	}
+
+	TWrapParam(const TArrayView<ValueType>& InArrayToWrap)
+		: TParam<ValueType>(
+			InArrayToWrap.GetData(), 
+			InArrayToWrap.Num(), 
+			((InArrayToWrap.Num()) > 1 ? FParam::EFlags::Batched : FParam::EFlags::None) | ((InArrayToWrap.Num() > 0) ? GetFlagsForConstructor(&InArrayToWrap[0]) : FParam::EFlags::None))
+	{
+	}
+
+protected:
+	void* GetDataForConstructor(ValueType* InValuePtrToWrap) const
+	{
+		using MutableValueType = typename TRemoveConst<ValueType>::Type;
+
+		// Support containers with Num() and GetData()
+		if constexpr (TModels<Private::CSizedContainerWithAccessibleDataAsRawPtr, ValueType>::Value)
+		{
+			return (void*)(InValuePtrToWrap->GetData());
+		}
+		else
+		{
+			return const_cast<MutableValueType*>(InValuePtrToWrap);
+		}
+	}
+
+	FParam::EFlags GetFlagsForConstructor(const ValueType* InValuePtrToWrap) const
+	{
+		FParam::EFlags NewFlags = FParam::EFlags::None;
+
+		// Add batched flags or not
+		if constexpr (TModels<Private::CSizedContainerWithAccessibleDataAsRawPtr, ValueType>::Value)
+		{
+			NewFlags |= (InValuePtrToWrap->Num() > 1) ? FParam::EFlags::Batched : FParam::EFlags::None;
+		}
+
+		// Add const flags or not
+		if constexpr (!TIsConst<ValueType>::Value)
+		{
+			NewFlags |= FParam::EFlags::Mutable;
+		}
+
+		return NewFlags;
+	}
+};
+
+// A typed param that owns it's own memory with size defined at compile time
+template<typename ValueType, int32 NumElem = 1>
+struct TParamValue : TParam<ValueType>
+{
+public:
+	TParamValue()
+		: TParam<ValueType>(&ValueArray[0], NumElem, GetFlagsForConstructor())
+	{
+	}
+
+	TParamValue(FParam::EFlags InFlags)
+		: TParam<ValueType>(&ValueArray[0], NumElem, InFlags)
+	{
+	}
+
+private:
+
+	FParam::EFlags GetFlagsForConstructor() const
+	{
+		FParam::EFlags NewFlags = FParam::EFlags::None;
+
+		NewFlags |= (NumElem > 1) ? FParam::EFlags::Batched : FParam::EFlags::None;
+
+		// Add const flags or not
+		if constexpr (!TIsConst<ValueType>::Value)
+		{
+			NewFlags |= FParam::EFlags::Mutable;
+		}
+
+		return NewFlags;
+	}
+
+	ValueType ValueArray[NumElem];
+};
+
+template<typename ValueType>
+struct TContextStorageParam : TWrapParam<ValueType>
+{
+public:
+	TContextStorageParam(ValueType* InValuePtrToWrap)
+		: TWrapParam<ValueType>(InValuePtrToWrap)
+	{
+		FParam::Flags |= FParam::EFlags::Stored;
+	}
+
+	TContextStorageParam(ValueType& InValueToWrap)
+		: TWrapParam<ValueType>(&InValueToWrap)
+	{
+		FParam::Flags |= FParam::EFlags::Stored;
+	}
+
+	TContextStorageParam(ValueType* InValuePtrToWrap, const int32 NumElements)
+		: TWrapParam<ValueType>(InValuePtrToWrap, NumElements)
+	{
+		FParam::Flags |= FParam::EFlags::Stored;
+	}
+
+	TContextStorageParam(const TArrayView<ValueType>& InArrayToWrap)
+		: TWrapParam<ValueType>(InArrayToWrap)
+	{
+		FParam::Flags |= FParam::EFlags::Stored;
+	}
+};
+
+// --- ---
+
+typedef int32 FParamHandle;
+constexpr FParamHandle InvalidParamHandle = -1;
+
+struct DATAINTERFACE_API FHParam
+{
+	FHParam() = default;
+	~FHParam();
+
+	FHParam(FHParam&& Other)
+		: FHParam()
+	{
+		Swap(*this, Other);
+	}
+
+	FHParam& operator= (FHParam&& Other)
+	{
+		Swap(*this, Other);
+		return *this;
+	}
+
+	FHParam(const FHParam& Other);
+
+	FHParam& operator= (const FHParam& Other)
+	{
+		FHParam Temp(Other); // note that the ref count is incremented here
+
+		Swap(*this, Temp);
+
+		return *this;
+	}
+
+protected:
+	friend struct FContext;
+	friend struct FParamStorage;
+
+	FHParam(FParamStorage* InOwnerStorage, FParamHandle InHandle);
+
+	FParamStorage* OwnerStorage = nullptr;
+	FParamHandle ParamHandle = InvalidParamHandle;
+
+private:
+};
 
+} // end namespace UE::DataInterface
diff --git a/Engine/Plugins/Experimental/DataInterface/Source/DataInterface/Public/DataInterfaceParamStorage.h b/Engine/Plugins/Experimental/DataInterface/Source/DataInterface/Public/DataInterfaceParamStorage.h
new file mode 100644
--- /dev/null
+++ b/Engine/Plugins/Experimental/DataInterface/Source/DataInterface/Public/DataInterfaceParamStorage.h
@@ -0,0 +1,222 @@
+// Copyright Epic Games, Inc. All Rights Reserved.
+
+#pragma once
+
+#include "CoreMinimal.h"
+#include "UObject/Class.h"
+#include "UObject/ScriptInterface.h"
+#include "DataInterfaceParam.h"
+
+namespace UE::DataInterface
+{
+
+typedef int32 FParamStorageHandle;
+constexpr FParamStorageHandle InvalidBlockHandle = -1;
+
+struct DATAINTERFACE_API FParamStorage
+{
+public:
+	typedef TPair<FParamStorageHandle, uint8*> TBlockDataPair;
+
+	FParamStorage(int32 InMaxParams, int32 InAllocatedMemorySize, int32 InMaxBlocks, bool InAllowGrowing = false);
+	~FParamStorage();
+
+
+	// *** Interface for param arena ***
+	// This prototype allows storing parameters into the storage using pre-allocated blocks of memory
+	// The owner is responsible for calculating the size and requesting / releasing the memory blocks
+	// Memory blocks are stored in stack LIFO way, last one has to be the first to be released (for simplicity)
+	// If a memoy block is released and it is not the last, it will be held until the ones after it are released
+
+	// Requests a block of RequestedBlockSize size
+	// Returns a pair with the block handle and the pointer to the memory buffer
+	TBlockDataPair RequestBlock(int32 RequestedBlockSize);
+
+	// Releases a block of memory from the arena, using a hanle to identify it
+	// If it is not the last of the stack, it will be held until the blocks after this one are released
+	void ReleaseBlock(FParamStorageHandle BlockHandle);
+
+	// Obtains the block of memory from a block handle
+	uint8* GetBlockMemory(FParamStorageHandle BlockHandle);
+
+
+	// *** Interface for direct Param storage ***
+	// This prototype allows storing parameters directly into the storage
+	// Parameters are stored stack like, in a LIFO approach
+	// If a parameter is released and it is not the last of the stack, it will be held until the parameters over that one are released
+	// This is to have a very simple memory management, might be removed on final implementation, based on requirements
+
+	// Adds a value to the storage (copying the data), returning a HParam
+	// If the total allocation size is <= than size of a void *, the value is stored directly in the FParam Data
+	// If not, additional memory is requested (TODO : better memory management)
+	template<typename ValueType>
+	FHParam AddValue(ValueType* Data, int32 NumElem, FParam::EFlags Flags)
+	{
+		check(ParamFreeIndex < (Parameters.Num() - 1)); // TODO : Growing ?
+
+		const int32 ParameterIndex = ParamFreeIndex;
+		ParamFreeIndex++;
+
+		const FParamType& ParamType = Private::TParamType<ValueType>::GetType();
+		const int32 ParamAllocSize = GetParamAllocSize(ParamType, NumElem);
+
+		FParam& Param = Parameters[ParameterIndex];
+
+		// Prepare the Param for the copy
+		Param.TypeId = ParamType.GetTypeId();
+		Param.Flags = Flags;
+		Param.NumElements = NumElem;
+		Param.Data = Data;
+
+		// If size is <= sizeof(void*), put the value directly at the pointer
+		if (ParamAllocSize <= sizeof(void*)) // TODO : see if we can get this constexpr
+		{
+			ParamMemoryHandles[ParameterIndex] = InvalidBlockHandle;
+
+			// Set target memory directly at Param.Data memory
+			ParamType.GetParamCopyFunction()(Param, static_cast<ValueType*>((void*)&Param.Data), ParamAllocSize);
+
+			Param.Flags |= FParam::EFlags::Embedded; // Added so it can be checked when retrieving (TODO : see if it can be done constexpr by type)
+		}
+		else
+		{
+			TBlockDataPair BlockData = RequestBlock(ParamAllocSize);
+			
+			ParamMemoryHandles[ParameterIndex] = BlockData.Key;
+
+			// Set target memory directly at BlockData memory
+			ParamType.GetParamCopyFunction()(Param, BlockData.Value, ParamAllocSize);
+			// Once the copy is done, update the Param
+			Param.Data = BlockData.Value;
+		}
+
+		RefCounts[ParameterIndex] = 1;
+
+		return FHParam(this, static_cast<FParamHandle>(ParameterIndex));
+	}
+
+	// Adds a reference (pointer) to the storage, returning a HParam
+	template<typename ValueType>
+	FHParam AddReference(ValueType* Data, int32 NumElem, FParam::EFlags Flags)
+	{
+		check(ParamFreeIndex < (Parameters.Num() - 1)); // TODO : Growing ?
+
+		const int32 ParameterIndex = ParamFreeIndex;
+		ParamFreeIndex++;
+
+		const FParamType& ParamType = Private::TParamType<ValueType>::GetType();
+
+		FParam& Param = Parameters[ParameterIndex];
+
+		Param.TypeId = ParamType.GetTypeId();
+		Param.Flags = Flags;
+		Param.NumElements = NumElem;
+		Param.Data = Data;
+
+		ParamMemoryHandles[ParameterIndex] = InvalidBlockHandle;
+
+		RefCounts[ParameterIndex] = 1;
+
+		return FHParam(this, static_cast<FParamHandle>(ParameterIndex));
+	}
+
+	// Obtains the Param associated to the HParam
+	FParam* GetParam(FParamHandle InParamHandle)
+	{
+		FParam* RetVal = nullptr;
+
+		const int32 ParameterIndex = InParamHandle;
+
+		if (ParameterIndex < ParamFreeIndex && RefCounts[ParameterIndex] > 0)
+		{
+			RetVal = &Parameters[ParameterIndex];
+		}
+
+		return RetVal;
+	}
+
+protected:
+	friend struct FHParam;
+
+	// Adds a reference count to a param, using the passed param handle
+	void IncRefCount(const FParamHandle InParamHandle)
+	{
+		const int32 ParameterIndex = InParamHandle;
+		check(ParameterIndex < ParamFreeIndex);
+		
+		check(RefCounts[ParameterIndex] > 0);
+		
+		RefCounts[ParameterIndex]++;
+	}
+
+	// Decrements a reference count to a param, using the passed param handle
+	// If the count reaches 0:
+	// - If it is the last param of the stack, it releases it and all the previous params with ref count == 0
+	// - If it is not the last one, returns without releasing (it will be done later)
+	void DecRefCount(const FParamHandle InParamHandle)
+	{
+		int32 ParameterIndex = InParamHandle;
+		check(ParameterIndex < ParamFreeIndex);
+
+		check(RefCounts[ParameterIndex] > 0);
+		
+		if (RefCounts[ParameterIndex] > 0)
+		{
+			RefCounts[ParameterIndex]--;
+
+			if (RefCounts[ParameterIndex] == 0)
+			{
+				// if it is the last parameter, release the memory and deal with the prev params with RefCount == 0, if any
+				if (ParameterIndex == ParamFreeIndex - 1)
+				{
+					// Recover free slots
+					while (ParameterIndex >= 0 && RefCounts[ParameterIndex] == 0)
+					{
+						Parameters[ParameterIndex] = FParam();
+						if (ParamMemoryHandles[ParameterIndex] != InvalidBlockHandle)
+						{
+							ReleaseBlock(ParamMemoryHandles[ParameterIndex]);
+							ParamMemoryHandles[ParameterIndex] = InvalidBlockHandle;
+						}
+
+						ParameterIndex--;
+						ParamFreeIndex--;
+					}
+				}
+			}
+		}
+	}
+
+	// Calculates the required size for a param, inclusing allignment
+	int32 GetParamAllocSize(const FParamType& ParamType, const int32 NumElem) const
+	{
+		const int32 ParamAlignment = ParamType.GetAlignment();
+		const int32 ParamSize = ParamType.GetSize();
+		const int32 ParamAllocSize = NumElem * Align(ParamSize, ParamAlignment);
+
+		return ParamAllocSize;
+	}
+
+	// --- context arena param storage prototype ---
+	TArray<int32> BlockOffets;
+	TArray<int32> BlockSizes;
+	TBitArray<> BlockFreeFlag;
+	uint8* RawMemory = nullptr;
+	int32 RawMemorySize = 0;
+	int32 CurrentBlocksSize = 0;
+	int32 FreeBlockIndex = 0;
+	
+	// --- shared param storage prototype ---
+	TArray<FParam> Parameters;
+	TArray<FParamStorageHandle> ParamMemoryHandles;
+	TArray<int16> RefCounts;
+	int32 ParamFreeIndex = 0;
+
+	bool AllowGrowing = false;
+
+private:
+	bool CheckBufferSize(int32 RequestedBlockSize);
+	bool CheckNumBlocks(int32 BlockIndex);
+};
+
+} // end namespace UE::DataInterface
diff --git a/Engine/Plugins/Experimental/DataInterface/Source/DataInterface/Public/DataInterfaceState.h b/Engine/Plugins/Experimental/DataInterface/Source/DataInterface/Public/DataInterfaceState.h
--- a/Engine/Plugins/Experimental/DataInterface/Source/DataInterface/Public/DataInterfaceState.h
+++ b/Engine/Plugins/Experimental/DataInterface/Source/DataInterface/Public/DataInterfaceState.h
@@ -22,12 +22,10 @@ enum class EStatePersistence : uint8
 // Container for data interface state
 struct DATAINTERFACE_API FState
 {
-	FState(uint32 InNumElements = 1, uint32 InChunkSize = 1)
+	FState(int32 InNumElements = 1)
 		: NumElements(InNumElements)
-		, ChunkSize(InChunkSize)
 	{
 		check(InNumElements > 0);
-		check(InChunkSize > 0);
 	}
 	
 private:
@@ -83,8 +81,7 @@ private:
 	
 	TMap<FInterfaceKeyWithIdAndStack, FRelevancyParam> RelevancyValueMap;
 	TMap<FInterfaceKeyWithIdAndStack, FParam> PermanentValueMap;
-	uint32 NumElements = 0;
-	uint32 ChunkSize = 0;
+	int32 NumElements = 0;
 };
 
 }
\ No newline at end of file
diff --git a/Engine/Plugins/Experimental/DataInterface/Source/DataInterface/Public/DataInterfaceTypes.h b/Engine/Plugins/Experimental/DataInterface/Source/DataInterface/Public/DataInterfaceTypes.h
--- a/Engine/Plugins/Experimental/DataInterface/Source/DataInterface/Public/DataInterfaceTypes.h
+++ b/Engine/Plugins/Experimental/DataInterface/Source/DataInterface/Public/DataInterfaceTypes.h
@@ -26,7 +26,12 @@ virtual FName GetReturnTypeNameImpl() const final override\
 	};
 
 #define IMPLEMENT_DATA_INTERFACE_PARAM_TYPE_INTERNAL(ValueType, Identifier) \
-	UE::DataInterface::FParamType::FRegistrar Identifier##_Registrar([](){ UE::DataInterface::FParamType::FRegistrar::RegisterType(Identifier##_DataInterfaceTypeInfo, UE::DataInterface::Private::TParamTypeHelper<ValueType>::GetStruct(), FName(#Identifier), UE::DataInterface::Private::TParamTypeHelper<ValueType>::GetSize(), UE::DataInterface::Private::TParamTypeHelper<ValueType>::GetAlignment()); });
+	UE::DataInterface::FParamType::FRegistrar Identifier##_Registrar([](){ UE::DataInterface::FParamType::FRegistrar::RegisterType( \
+		Identifier##_DataInterfaceTypeInfo, \
+		UE::DataInterface::Private::TParamTypeHelper<ValueType>::GetStruct(), \
+		FName(#Identifier), UE::DataInterface::Private::TParamTypeHelper<ValueType>::GetSize(), \
+		UE::DataInterface::Private::TParamTypeHelper<ValueType>::GetAlignment(), \
+		UE::DataInterface::Private::TParamTypeHelper<ValueType>::GetParamCloneFunction()); });
 
 // Types that are passed between data interfaces (i.e. exist in in public headers) should declare via this mechanism
 #define DECLARE_DATA_INTERFACE_PARAM_TYPE(APIType, ValueType, Identifier) \
@@ -46,7 +51,12 @@ virtual FName GetReturnTypeNameImpl() const final override\
 // If they are exposed in a public header please use the DECLARE_DATA_INTERFACE_PARAM_TYPE & IMPLEMENT_DATA_INTERFACE_PARAM_TYPE macros
 #define IMPLEMENT_DATA_INTERFACE_STATE_TYPE(ValueType, Identifier) \
 	static UE::DataInterface::FParamType Identifier##_DataInterfaceTypeInfo; \
-	static UE::DataInterface::FParamType::FRegistrar Identifier##_Registrar([](){ UE::DataInterface::FParamType::FRegistrar::RegisterType(Identifier##_DataInterfaceTypeInfo, UE::DataInterface::Private::TParamTypeHelper<ValueType>::GetStruct(), FName(#Identifier), UE::DataInterface::Private::TParamTypeHelper<ValueType>::GetSize(), UE::DataInterface::Private::TParamTypeHelper<ValueType>::GetAlignment()); }); \
+	static UE::DataInterface::FParamType::FRegistrar Identifier##_Registrar([](){ UE::DataInterface::FParamType::FRegistrar::RegisterType( \
+		Identifier##_DataInterfaceTypeInfo, \
+		UE::DataInterface::Private::TParamTypeHelper<ValueType>::GetStruct(), \
+		FName(#Identifier), UE::DataInterface::Private::TParamTypeHelper<ValueType>::GetSize(), \
+		UE::DataInterface::Private::TParamTypeHelper<ValueType>::GetAlignment(), \
+		UE::DataInterface::Private::TParamTypeHelper<ValueType>::GetParamCloneFunction()); }); \
 	DECLARE_DATA_INTERFACE_PARAM_TYPE_INTERNAL(ValueType, Identifier);
 
 namespace UE::DataInterface
diff --git a/Engine/Plugins/Experimental/DataInterface/Source/DataInterface/Public/IDataInterface.h b/Engine/Plugins/Experimental/DataInterface/Source/DataInterface/Public/IDataInterface.h
--- a/Engine/Plugins/Experimental/DataInterface/Source/DataInterface/Public/IDataInterface.h
+++ b/Engine/Plugins/Experimental/DataInterface/Source/DataInterface/Public/IDataInterface.h
@@ -25,6 +25,18 @@ class DATAINTERFACE_API IDataInterface
 	GENERATED_BODY()
 
 public:
+	/**
+	 * Gets data using context to retrieve input parameters and result
+	 * @return false if type are incompatible, or if nested calls fail, otherwise true
+	 */
+	bool GetData(const UE::DataInterface::FContext& Context) const;
+
+	/**
+	 * Gets data and stores the value in OutResult, checking whether types are compatible
+	* @return false if nested calls fail, otherwise true
+	*/
+	bool GetDataChecked(const UE::DataInterface::FContext& Context) const;
+
 	/**
 	 * Gets data and stores the value in OutResult.
 	 * @return false if type are incompatible, or if nested calls fail, otherwise true
diff --git a/Engine/Plugins/Experimental/DataInterface/Source/DataInterface/Public/IDataInterfaceParameters.h b/Engine/Plugins/Experimental/DataInterface/Source/DataInterface/Public/IDataInterfaceParameters.h
--- a/Engine/Plugins/Experimental/DataInterface/Source/DataInterface/Public/IDataInterfaceParameters.h
+++ b/Engine/Plugins/Experimental/DataInterface/Source/DataInterface/Public/IDataInterfaceParameters.h
@@ -28,6 +28,9 @@ class DATAINTERFACE_API IDataInterfaceParameters
 	// Get a parameter by name
 	virtual bool GetParameter(FName InKey, UE::DataInterface::FParam& OutParam) = 0;
 
+	// Get a parameter pointer by name
+	virtual const UE::DataInterface::FParam* GetParameter(FName InKey) = 0;
+
 	// Default utility implementation for UObjects
 	// Treats key as a property name or function
 	static bool GetParameterForObject(UObject* InObject, FName InKey, UE::DataInterface::FParam& OutParam);
diff --git a/Engine/Plugins/Experimental/DataInterfaceGraph/Source/DataInterfaceGraph/Private/RigUnit_DataInterfaceParameter.cpp b/Engine/Plugins/Experimental/DataInterfaceGraph/Source/DataInterfaceGraph/Private/RigUnit_DataInterfaceParameter.cpp
--- a/Engine/Plugins/Experimental/DataInterfaceGraph/Source/DataInterfaceGraph/Private/RigUnit_DataInterfaceParameter.cpp
+++ b/Engine/Plugins/Experimental/DataInterfaceGraph/Source/DataInterfaceGraph/Private/RigUnit_DataInterfaceParameter.cpp
@@ -40,7 +40,7 @@ FRigUnit_DataInterface_Float_Execute()
 		const FContext& DataInterfaceContext = DataInterfaceExecuteContext.GetContext();
 
 		// Wrap the internal result we are going to be writing to 
-		TWrapParam<float> CallResult(DataInterfaceContext, &Result);
+		TWrapParam<float> CallResult(&Result);
 
 		// Call the interface
 		DataInterfaceExecuteContext.SetResult(UE::DataInterface::GetDataSafe(DataInterface, DataInterfaceContext, CallResult));
@@ -59,7 +59,7 @@ FRigUnit_DataInterface_Pose_Execute()
 		const FContext& DataInterfaceContext = DataInterfaceUnitContext.DataInterfaceContext;
 
 		// Wrap the internal result we are going to be writing to 
-		TWrapParam<FDataInterfaceExecuteContext> CallResult(DataInterfaceContext, &Result);
+		TWrapParam<FDataInterfaceExecuteContext> CallResult(&Result);
 
 		// Call the interface
 		DataInterfaceUnitContext.bResult &= UE::DataInterface::GetDataSafe(PoseInterface, DataInterfaceContext, CallResult);
