[DataInterface]
* Attempt at fixing Hololens compile error

diff --git a/Engine/Plugins/Experimental/DataInterface/Source/DataInterface/Private/DataInterfaceParamStorage.cpp b/Engine/Plugins/Experimental/DataInterface/Source/DataInterface/Private/DataInterfaceParamStorage.cpp
--- a/Engine/Plugins/Experimental/DataInterface/Source/DataInterface/Private/DataInterfaceParamStorage.cpp
+++ b/Engine/Plugins/Experimental/DataInterface/Source/DataInterface/Private/DataInterfaceParamStorage.cpp
@@ -107,8 +107,9 @@ uint8* FParamStorage::GetBlockMemory(FParamStorageHandle BlockHandle)
 
 	if (ensure(BlockIndex >= 0 && BlockIndex < FreeBlockIndex))
 	{
-		ensure(BlockFreeFlag[BlockIndex] == false);
-		BlockMemory = (BlockFreeFlag[BlockIndex] == false) ? RawMemory + BlockOffets[BlockIndex] : nullptr;
+		bool IsSet = BlockFreeFlag[BlockIndex];
+		ensure(IsSet == false);
+		BlockMemory = (IsSet == false) ? RawMemory + BlockOffets[BlockIndex] : nullptr;
 	}
 
 	return BlockMemory;
